UBT: Add a custom wrapper to invoke the linker on Windows, to filter out intermittent "Unexpected PDB error; OK" messages and restart the link.

diff --git a/Engine/Build/InstalledEngineFilters.xml b/Engine/Build/InstalledEngineFilters.xml
--- a/Engine/Build/InstalledEngineFilters.xml
+++ b/Engine/Build/InstalledEngineFilters.xml
@@ -209,6 +209,7 @@
 
 			<!-- Build utilities -->
 			Engine/Build/Windows/cl-filter/cl-filter.exe
+			Engine/Build/Windows/link-filter/link-filter.exe
 			Engine/Source/ThirdParty/IntelISPC/bin/Windows/ispc.exe
 
 			<!-- Default resource files -->
diff --git a/Engine/Extras/Windows/link-filter/link-filter.cpp b/Engine/Extras/Windows/link-filter/link-filter.cpp
new file mode 100644
--- /dev/null
+++ b/Engine/Extras/Windows/link-filter/link-filter.cpp
@@ -0,0 +1,181 @@
+// Copyright 1998-2019 Epic Games, Inc. All Rights Reserved.
+
+/**
+ * This program is designed to execute the Visual C++ linker (link.exe) and restart if it sees a spurious "Unexpected PDB error; OK (0)"
+ * message in the output.
+ */
+
+#include <windows.h>
+#include <stdio.h>
+
+static const int ERROR_EXIT_CODE = 1;
+
+void PrintUsage()
+{
+	wprintf(L"Usage: link-filter.exe -linker=<linker-file> -- <child command line>\n");
+}
+
+int wmain(int ArgC, const wchar_t* ArgV[])
+{
+	// Get the child command line.
+	wchar_t* ChildCmdLine = wcsstr(::GetCommandLineW(), L" -- ");
+	if (ChildCmdLine == nullptr)
+	{
+		wprintf(L"ERROR: No child command line specified.\n");
+		PrintUsage();
+		return ERROR_EXIT_CODE;
+	}
+	ChildCmdLine += 4;
+	
+	// Retry until we get a valid run of the linker
+	DWORD ExitCode;
+	for (size_t Retry = 0; Retry < 3; Retry++)
+	{
+		// Create the child process
+		PROCESS_INFORMATION ProcessInfo;
+		ZeroMemory(&ProcessInfo, sizeof(ProcessInfo));
+
+		SECURITY_ATTRIBUTES SecurityAttributes;
+		ZeroMemory(&SecurityAttributes, sizeof(SecurityAttributes));
+		SecurityAttributes.bInheritHandle = TRUE;
+
+		HANDLE StdOutReadHandle;
+		HANDLE StdOutWriteHandle;
+		if (CreatePipe(&StdOutReadHandle, &StdOutWriteHandle, &SecurityAttributes, 0) == 0)
+		{
+			wprintf(L"ERROR: Unable to create output pipe for child process\n");
+			return ERROR_EXIT_CODE;
+		}
+
+		HANDLE StdErrWriteHandle;
+		if (DuplicateHandle(GetCurrentProcess(), StdOutWriteHandle, GetCurrentProcess(), &StdErrWriteHandle, 0, true, DUPLICATE_SAME_ACCESS) == 0)
+		{
+			wprintf(L"ERROR: Unable to create stderr pipe handle for child process\n");
+			return ERROR_EXIT_CODE;
+		}
+
+		// Create the new process as suspended, so we can modify it before it starts executing (and potentially preempting us)
+		STARTUPINFO StartupInfo;
+		ZeroMemory(&StartupInfo, sizeof(StartupInfo));
+		StartupInfo.cb = sizeof(StartupInfo);
+		StartupInfo.hStdInput = NULL;
+		StartupInfo.hStdOutput = StdOutWriteHandle;
+		StartupInfo.hStdError = StdErrWriteHandle;
+		StartupInfo.dwFlags = STARTF_USESTDHANDLES;
+
+		DWORD ProcessCreationFlags = GetPriorityClass(GetCurrentProcess());
+		if (CreateProcessW(NULL, ChildCmdLine, NULL, NULL, TRUE, ProcessCreationFlags, NULL, NULL, &StartupInfo, &ProcessInfo) == 0)
+		{
+			wprintf(L"ERROR: Unable to create child process\n");
+			return ERROR_EXIT_CODE;
+		}
+
+		// Close the startup thread handle; we don't need it.
+		CloseHandle(ProcessInfo.hThread);
+
+		// Close the write ends of the handle. We don't want any other process to be able to inherit these.
+		CloseHandle(StdOutWriteHandle);
+		CloseHandle(StdErrWriteHandle);
+
+		// Whether we detected an invalid run
+		bool bRestartLink = false;
+
+		// Pipe the output to stdout
+		char Buffer[1024];
+		size_t BufferSize = 0;
+		for (;;)
+		{
+			// Read the next chunk of data from the output stream
+			DWORD BytesRead = 0;
+			if (BufferSize < sizeof(Buffer))
+			{
+				if (ReadFile(StdOutReadHandle, Buffer + BufferSize, (DWORD)(sizeof(Buffer) - BufferSize), &BytesRead, NULL))
+				{
+					BufferSize += BytesRead;
+				}
+				else if(GetLastError() != ERROR_BROKEN_PIPE)
+				{
+					wprintf(L"ERROR: Unable to read data from child process (%08x)", GetLastError());
+				}
+				else if (BufferSize == 0)
+				{
+					break;
+				}
+			}
+
+			// Parse individual lines from the output
+			size_t LineStart = 0;
+			while(LineStart < BufferSize)
+			{
+				// Find the end of this line
+				size_t LineEnd = LineStart;
+				while (LineEnd < BufferSize && Buffer[LineEnd] != '\n')
+				{
+					LineEnd++;
+				}
+
+				// If we didn't reach a line terminator, and we can still read more data, clear up some space and try again
+				if (LineEnd == BufferSize && !(LineStart == 0 && BytesRead == 0) && !(LineStart == 0 && BufferSize == sizeof(Buffer)))
+				{
+					break;
+				}
+
+				// Skip past the EOL marker
+				if (LineEnd < BufferSize && Buffer[LineEnd] == '\n')
+				{
+					LineEnd++;
+				}
+
+				// Check if this line contains the spurious error
+				const char ErrorText[] = "Unexpected PDB error; OK (0)";
+				size_t ErrorTextLen = sizeof(ErrorText) - 1;
+				for (size_t Idx = LineStart; Idx + ErrorTextLen <= LineEnd; Idx++)
+				{
+					if (memcmp(Buffer + Idx, ErrorText, ErrorTextLen) == 0)
+					{
+						wprintf(L"NOTE: Detected '%S' string in output. Restarting link.\n", ErrorText);
+						bRestartLink = true;
+						LineStart = LineEnd;
+						break;
+					}
+				}
+
+				// If we didn't write anything out, write it to stdout
+				if(LineStart < LineEnd)
+				{
+					DWORD BytesWritten;
+					WriteFile(GetStdHandle(STD_OUTPUT_HANDLE), Buffer + LineStart, (DWORD)(LineEnd - LineStart), &BytesWritten, NULL);
+				}
+
+				// Move to the next line
+				LineStart = LineEnd;
+			}
+
+			// Shuffle everything down
+			if (LineStart > 0)
+			{
+				memmove(Buffer, Buffer + LineStart, BufferSize - LineStart);
+				BufferSize -= LineStart;
+			}
+		}
+
+		WaitForSingleObject(ProcessInfo.hProcess, INFINITE);
+
+		if (!GetExitCodeProcess(ProcessInfo.hProcess, &ExitCode))
+		{
+			ExitCode = (DWORD)ERROR_EXIT_CODE;
+		}
+
+		CloseHandle(ProcessInfo.hProcess);
+
+		if (bRestartLink)
+		{
+			ExitCode = ERROR_EXIT_CODE;
+		}
+		else
+		{
+			break;
+		}
+	}
+	return ExitCode;
+}
diff --git a/Engine/Source/Programs/UnrealBuildTool/Platform/Windows/VCToolChain.cs b/Engine/Source/Programs/UnrealBuildTool/Platform/Windows/VCToolChain.cs
--- a/Engine/Source/Programs/UnrealBuildTool/Platform/Windows/VCToolChain.cs
+++ b/Engine/Source/Programs/UnrealBuildTool/Platform/Windows/VCToolChain.cs
@@ -1800,8 +1800,16 @@ namespace UnrealBuildTool
 			Action LinkAction = new Action(ActionType.Link);
 			LinkAction.CommandDescription = "Link";
 			LinkAction.WorkingDirectory = UnrealBuildTool.EngineSourceDirectory;
-			LinkAction.CommandPath = bIsBuildingLibraryOrImportLibrary ? EnvVars.LibraryManagerPath : EnvVars.LinkerPath;
-			LinkAction.CommandArguments = String.Format("@\"{0}\"", ResponseFileName);
+			if(bIsBuildingLibraryOrImportLibrary)
+			{
+				LinkAction.CommandPath = EnvVars.LibraryManagerPath;
+				LinkAction.CommandArguments = String.Format("@\"{0}\"", ResponseFileName);
+			}
+			else
+			{
+				LinkAction.CommandPath = FileReference.Combine(UnrealBuildTool.EngineDirectory, "Build", "Windows", "link-filter", "link-filter.exe");
+				LinkAction.CommandArguments = String.Format("-- \"{0}\" @\"{1}\"", EnvVars.LinkerPath, ResponseFileName);
+			}
 			LinkAction.ProducedItems.AddRange(ProducedItems);
 			LinkAction.PrerequisiteItems.AddRange(PrerequisiteItems);
 			LinkAction.StatusDescription = Path.GetFileName(OutputFile.AbsolutePath);
