Move cl-filter into a Windows subfolder.

diff --git a/Engine/Extras/Windows/cl-filter/cl-filter.cpp b/Engine/Extras/Windows/cl-filter/cl-filter.cpp
new file mode 100644
--- /dev/null
+++ b/Engine/Extras/Windows/cl-filter/cl-filter.cpp
@@ -0,0 +1,443 @@
+// Copyright 1998-2019 Epic Games, Inc. All Rights Reserved.
+
+/**
+ * This program is designed to execute the Visual C++ compiler (cl.exe) and filter off any output lines from the /showIncludes directive
+ * into a separate file for dependency checking. GCC/Clang have a specific option for this, whereas MSVC does not.
+ */
+
+#include <windows.h>
+#include <stdio.h>
+#include <string>
+#include <vector>
+#include <set>
+
+void GetLocalizedIncludePrefixes(const wchar_t* CompilerPath, std::vector<std::vector<char>>& Prefixes);
+
+#define EXIT_CODE_SUCCESS 0
+#define EXIT_CODE_FAILURE -1
+
+#define OUTPUT_FILENAME_INDEX 1
+#define TIMING_FILENAME_INDEX 2
+#define COMPILER_FILENAME_INDEX_OFFSET 2
+#define MINIMUM_VALID_COMMAND_LINE_ARGS 4
+
+void PrintUsage()
+{
+	wprintf(L"Usage: cl-filter.exe -dependencies=<dependencies-file> -compiler=<compiler-file> [Optional Args] -- <child command line>\n");
+	wprintf(L"Optional Args:\n");
+	wprintf(L"\t-timing=<timing-file>\tThe file to write any timing data to.\n");
+}
+
+struct CommandLineArgs
+{
+	const wchar_t*	OutputFileName = nullptr;
+	const wchar_t*	TimingFileName = nullptr;
+	const wchar_t*	CompilerFileName = nullptr;
+};
+
+CommandLineArgs* ParseCommandLineArgs(int ArgC, const wchar_t* ArgV[])
+{
+	CommandLineArgs* ParsedArgs = new CommandLineArgs();
+	for (int i = 1; i < ArgC; ++i)
+	{
+		// If this arg is the splitter, we're done parsing args.
+		if (wcscmp(ArgV[i], L"--") == 0)
+		{
+			break;
+		}
+
+		// If it starts with a -, parse the arg.
+		if (ArgV[i][0] == '-')
+		{
+			if (wcsstr(ArgV[i], L"-dependencies=") != nullptr)
+			{
+				ParsedArgs->OutputFileName = ArgV[i] + 14;
+			}
+			else if (wcsstr(ArgV[i], L"-compiler=") != nullptr)
+			{
+				ParsedArgs->CompilerFileName = ArgV[i] + 10;
+			}
+			else if (wcsstr(ArgV[i], L"-timing=") != nullptr)
+			{
+				ParsedArgs->TimingFileName = ArgV[i] + 8;
+			}
+		}
+	}
+
+	return ParsedArgs;
+}
+
+int wmain(int ArgC, const wchar_t* ArgV[])
+{
+	wchar_t* CmdLine = ::GetCommandLineW();
+	CommandLineArgs* ParsedArgs = ParseCommandLineArgs(ArgC, ArgV);
+	if (ParsedArgs->OutputFileName == nullptr)
+	{
+		wprintf(L"ERROR: No output file name was specified! (%s)\n\n", CmdLine);
+		PrintUsage();
+		return EXIT_CODE_FAILURE;
+	}
+
+	if (ParsedArgs->CompilerFileName == nullptr)
+	{
+		wprintf(L"ERROR: No compiler file name was found! (%s)\n\n", CmdLine);
+		PrintUsage();
+		return EXIT_CODE_FAILURE;
+	}
+
+	// Get the child command line.
+	wchar_t* ChildCmdLine = wcsstr(CmdLine, L" -- ");
+	if (ChildCmdLine == nullptr)
+	{
+		wprintf(L"ERROR: Unable to find child command line! (%s)\n\n", CmdLine);
+		PrintUsage();
+		return EXIT_CODE_FAILURE;
+	}
+	ChildCmdLine += 4;
+	
+	// Get all the possible localized string prefixes for /showIncludes output
+	std::vector<std::vector<char>> LocalizedIncludePrefixes;
+	GetLocalizedIncludePrefixes(ParsedArgs->CompilerFileName, LocalizedIncludePrefixes);
+
+	// Create the child process
+	PROCESS_INFORMATION ProcessInfo;
+	ZeroMemory(&ProcessInfo, sizeof(ProcessInfo));
+
+	SECURITY_ATTRIBUTES SecurityAttributes;
+	ZeroMemory(&SecurityAttributes, sizeof(SecurityAttributes));
+	SecurityAttributes.bInheritHandle = TRUE;
+
+	HANDLE StdOutReadHandle;
+	HANDLE StdOutWriteHandle;
+	if (CreatePipe(&StdOutReadHandle, &StdOutWriteHandle, &SecurityAttributes, 0) == 0)
+	{
+		wprintf(L"ERROR: Unable to create output pipe for child process\n");
+		return EXIT_CODE_FAILURE;
+	}
+
+	HANDLE StdErrWriteHandle;
+	if (DuplicateHandle(GetCurrentProcess(), StdOutWriteHandle, GetCurrentProcess(), &StdErrWriteHandle, 0, true, DUPLICATE_SAME_ACCESS) == 0)
+	{
+		wprintf(L"ERROR: Unable to create stderr pipe handle for child process\n");
+		return EXIT_CODE_FAILURE;
+	}
+
+	// Create the new process as suspended, so we can modify it before it starts executing (and potentially preempting us)
+	STARTUPINFO StartupInfo;
+	ZeroMemory(&StartupInfo, sizeof(StartupInfo));
+	StartupInfo.cb = sizeof(StartupInfo);
+	StartupInfo.hStdInput = NULL;
+	StartupInfo.hStdOutput = StdOutWriteHandle;
+	StartupInfo.hStdError = StdErrWriteHandle;
+	StartupInfo.dwFlags = STARTF_USESTDHANDLES;
+
+	DWORD ProcessCreationFlags = GetPriorityClass(GetCurrentProcess());
+	if (CreateProcessW(NULL, ChildCmdLine, NULL, NULL, TRUE, ProcessCreationFlags, NULL, NULL, &StartupInfo, &ProcessInfo) == 0)
+	{
+		wprintf(L"ERROR: Unable to create child process\n");
+		return EXIT_CODE_FAILURE;
+	}
+
+	// Close the startup thread handle; we don't need it.
+	CloseHandle(ProcessInfo.hThread);
+
+	// Close the write ends of the handle. We don't want any other process to be able to inherit these.
+	CloseHandle(StdOutWriteHandle);
+	CloseHandle(StdErrWriteHandle);
+
+	// Delete the output file.
+	DeleteFileW(ParsedArgs->OutputFileName);
+
+	// Delete the timing file if needed.
+	if (ParsedArgs->TimingFileName != nullptr)
+	{
+		DeleteFileW(ParsedArgs->TimingFileName);
+	}
+
+	// Get the path to a temporary output filename.
+	std::wstring TempOutputFileName(ParsedArgs->OutputFileName);
+	TempOutputFileName += L".tmp";
+
+	// Get the path to a temporary timing filename.
+	std::wstring TempTimingFileName(ParsedArgs->TimingFileName == nullptr ? L"" : ParsedArgs->TimingFileName);
+	TempTimingFileName += L".tmp";
+
+	// Create a file to contain the dependency list.
+	HANDLE OutputFile = CreateFileW(TempOutputFileName.c_str(), GENERIC_WRITE, FILE_SHARE_READ, NULL, CREATE_ALWAYS, FILE_ATTRIBUTE_NORMAL, NULL);
+	if(OutputFile == INVALID_HANDLE_VALUE)
+	{
+		wprintf(L"ERROR: Unable to open %s for output", TempOutputFileName.c_str());
+		return EXIT_CODE_FAILURE;
+	}
+
+	// If needed, create a file to contain unparsed timing info.
+	HANDLE TimingFile = INVALID_HANDLE_VALUE;
+	if (ParsedArgs->TimingFileName != nullptr)
+	{
+		TimingFile = CreateFileW(TempTimingFileName.c_str(), GENERIC_WRITE, FILE_SHARE_READ, NULL, CREATE_ALWAYS, FILE_ATTRIBUTE_NORMAL, NULL);
+		if (TimingFile == INVALID_HANDLE_VALUE)
+		{
+			wprintf(L"ERROR: Unable to open %s for output", TempTimingFileName.c_str());
+			return EXIT_CODE_FAILURE;
+		}
+	}
+
+	// Pipe the output to stdout
+	char Buffer[1024];
+	size_t BufferSize = 0;
+
+	bool InTimingInfo = false;
+	int TimingLinesToCapture = 0;
+	for (;;)
+	{
+		// Read the next chunk of data from the output stream
+		DWORD BytesRead = 0;
+		if (BufferSize < sizeof(Buffer))
+		{
+			if (ReadFile(StdOutReadHandle, Buffer + BufferSize, (DWORD)(sizeof(Buffer) - BufferSize), &BytesRead, NULL))
+			{
+				BufferSize += BytesRead;
+			}
+			else if(GetLastError() != ERROR_BROKEN_PIPE)
+			{
+				wprintf(L"ERROR: Unable to read data from child process (%08x)", GetLastError());
+			}
+			else if (BufferSize == 0)
+			{
+				break;
+			}
+		}
+
+		// Parse individual lines from the output
+		size_t LineStart = 0;
+		while(LineStart < BufferSize)
+		{
+			// Find the end of this line
+			size_t LineEnd = LineStart;
+			while (LineEnd < BufferSize && Buffer[LineEnd] != '\n')
+			{
+				LineEnd++;
+			}
+
+			// If we didn't reach a line terminator, and we can still read more data, clear up some space and try again
+			if (LineEnd == BufferSize && !(LineStart == 0 && BytesRead == 0) && !(LineStart == 0 && BufferSize == sizeof(Buffer)))
+			{
+				break;
+			}
+
+			// Skip past the EOL marker
+			if (LineEnd < BufferSize && Buffer[LineEnd] == '\n')
+			{
+				LineEnd++;
+			}
+
+			// Filter out any lines that have the "Note: including file: " prefix.
+			for (const std::vector<char>& LocalizedPrefix : LocalizedIncludePrefixes)
+			{
+				if (memcmp(Buffer + LineStart, LocalizedPrefix.data(), LocalizedPrefix.size() - 1) == 0)
+				{
+					size_t FileNameStart = LineStart + LocalizedPrefix.size() - 1;
+					while (FileNameStart < LineEnd && isspace(Buffer[FileNameStart]))
+					{
+						FileNameStart++;
+					}
+
+					DWORD BytesWritten;
+					WriteFile(OutputFile, Buffer + FileNameStart, (DWORD)(LineEnd - FileNameStart), &BytesWritten, NULL);
+					LineStart = LineEnd;
+					break;
+				}
+			}
+
+			// Handling timing info parsing if needed.
+			if (ParsedArgs->TimingFileName != nullptr && LineStart < LineEnd)
+			{
+				// See if we've moved into a block of timing information. Order is always Headers -> Classes -> Functions, so to keep for loops to a minimum,
+				//   handle by what our current state is.
+				const char StartTimingInfoText[] = "Include Headers:";
+				
+				if (!InTimingInfo && memcmp(Buffer + LineStart, StartTimingInfoText, 16) == 0)
+				{
+					InTimingInfo = true;	
+				}
+
+				if (InTimingInfo)
+				{
+					DWORD BytesWritten;
+					WriteFile(TimingFile, Buffer + LineStart, (DWORD)(LineEnd - LineStart), &BytesWritten, NULL);
+					LineStart = LineEnd;
+					continue;
+				}
+			}
+
+			// If we didn't write anything out, write it to stdout
+			if(LineStart < LineEnd)
+			{
+				DWORD BytesWritten;
+				WriteFile(GetStdHandle(STD_OUTPUT_HANDLE), Buffer + LineStart, (DWORD)(LineEnd - LineStart), &BytesWritten, NULL);
+			}
+
+			// Move to the next line
+			LineStart = LineEnd;
+		}
+
+		// Shuffle everything down
+		memmove(Buffer, Buffer + LineStart, BufferSize - LineStart);
+		BufferSize -= LineStart;
+	}
+
+	WaitForSingleObject(ProcessInfo.hProcess, INFINITE);
+
+	DWORD ExitCode;
+	if (!GetExitCodeProcess(ProcessInfo.hProcess, &ExitCode))
+	{
+		ExitCode = (DWORD)EXIT_CODE_FAILURE;
+	}
+
+	CloseHandle(OutputFile);
+
+	if (ExitCode == 0 && !MoveFileW(TempOutputFileName.c_str(), ParsedArgs->OutputFileName))
+	{
+		wprintf(L"ERROR: Unable to rename %s to %s\n", TempOutputFileName.c_str(), ParsedArgs->OutputFileName);
+		ExitCode = 1;
+	}
+
+	if (ParsedArgs->TimingFileName != nullptr)
+	{
+		CloseHandle(TimingFile);
+		if (ExitCode == 0 && !MoveFileW(TempTimingFileName.c_str(), ParsedArgs->TimingFileName))
+		{
+			wprintf(L"ERROR: Unable to rename %s to %s\n", TempTimingFileName.c_str(), ParsedArgs->TimingFileName);
+			ExitCode = 1;
+		}
+	}
+
+	return ExitCode;
+}
+
+static std::string FindAndReplace(std::string Input, const std::string& FindStr, const std::string& ReplaceStr)
+{
+	size_t Start = 0;
+	for (;;)
+	{
+		size_t Offset = Input.find(FindStr, Start);
+		if(Offset == std::wstring::npos)
+		{
+			break;
+		}
+
+		Input.replace(Offset, FindStr.size(), ReplaceStr);
+		Start = Offset + ReplaceStr.size();
+	}
+	return Input;
+}
+
+bool GetLocalizedIncludePrefix(UINT CodePage, const wchar_t* LibraryPath, HMODULE LibraryHandle, std::vector<char>& Prefix)
+{
+	static const unsigned int ResourceId = 408;
+
+	// Read the string from the library
+	wchar_t Text[512];
+	if(LoadStringW(LibraryHandle, ResourceId, Text, sizeof(Text) / sizeof(Text[0])) == 0)
+	{
+		wprintf(L"WARNING: unable to read string %d from %s\n", ResourceId, LibraryPath);
+		return false;
+	}
+
+	// Find the end of the prefix 
+	wchar_t* TextEnd = wcsstr(Text, L"%s%s");
+	if (TextEnd == nullptr)
+	{
+		wprintf(L"WARNING: unable to find substitution markers in format string '%s' (%s)", Text, LibraryPath);
+		return false;
+	}
+
+	// Figure out how large the buffer needs to be to hold the MBCS version
+	int Length = WideCharToMultiByte(CP_ACP, 0, Text, (int)(TextEnd - Text), NULL, 0, NULL, NULL);
+	if (Length == 0)
+	{
+		wprintf(L"WARNING: unable to query size for MBCS output buffer (input text '%s', library %s)", Text, LibraryPath);
+		return false;
+	}
+
+	// Resize the output buffer with space for a null terminator
+	Prefix.resize(Length + 1);
+
+	// Get the multibyte text
+	int Result = WideCharToMultiByte(CodePage, 0, Text, (int)(TextEnd - Text), Prefix.data(), Length, NULL, NULL);
+	if (Result == 0)
+	{
+		wprintf(L"WARNING: unable to get MBCS string (input text '%s', library %s)", Text, LibraryPath);
+		return false;
+	}
+
+	return true;
+}
+
+// Language packs for Visual Studio contain localized strings for the "Note: including file:" prefix we expect to see when running the compiler
+// with the /showIncludes option. Enumerate all the possible languages that may be active, and build up an array of possible prefixes. We'll treat
+// any of them as markers for included files.
+void GetLocalizedIncludePrefixes(const wchar_t* CompilerPath, std::vector<std::vector<char>>& Prefixes)
+{
+	// Get all the possible locale id's. Include en-us by default.
+	std::set<std::wstring> LocaleIds;
+	LocaleIds.insert(L"1033");
+
+	// The user default locale id
+	wchar_t LocaleIdString[20];
+	wsprintf(LocaleIdString, L"%d", GetUserDefaultLCID());
+	LocaleIds.insert(LocaleIdString);
+
+	// The system default locale id
+	wsprintf(LocaleIdString, L"%d", GetSystemDefaultLCID());
+	LocaleIds.insert(LocaleIdString);
+
+	// The Visual Studio locale setting
+	static const size_t VsLangMaxLen = 256;
+	wchar_t VsLangEnv[VsLangMaxLen];
+	if (GetEnvironmentVariableW(L"VSLANG", VsLangEnv, VsLangMaxLen) != 0)
+	{
+		LocaleIds.insert(VsLangEnv);
+	}
+
+	// Find the directory containing the compiler path
+	size_t CompilerDirLen = wcslen(CompilerPath);
+	while (CompilerDirLen > 0 && CompilerPath[CompilerDirLen - 1] != '/' && CompilerPath[CompilerDirLen - 1] != '\\')
+	{
+		CompilerDirLen--;
+	}
+
+	// Always add the en-us prefix. We'll validate that this is correct if we have an en-us resource file, but it gives us something to fall back on.
+	const char EnglishText[] = "Note: including file:";
+	Prefixes.emplace_back(EnglishText, strchr(EnglishText, 0) + 1);
+	
+	// Get the default console codepage
+	UINT CodePage = GetConsoleOutputCP();
+
+	// Loop through all the possible locale ids and try to find the localized string for each
+	for (const std::wstring& LocaleId : LocaleIds)
+	{
+		std::wstring ResourceFile;
+		ResourceFile.assign(CompilerPath, CompilerPath + CompilerDirLen);
+		ResourceFile.append(LocaleId);
+		ResourceFile.append(L"\\clui.dll");
+
+		HMODULE LibraryHandle = LoadLibraryExW(ResourceFile.c_str(), 0, LOAD_LIBRARY_AS_DATAFILE | LOAD_LIBRARY_AS_IMAGE_RESOURCE);
+		if (LibraryHandle != nullptr)
+		{
+			std::vector<char> Prefix;
+			if (GetLocalizedIncludePrefix(CodePage, ResourceFile.c_str(), LibraryHandle, Prefix))
+			{
+				if (wcscmp(LocaleId.c_str(), L"1033") != 0)
+				{
+					Prefixes.push_back(std::move(Prefix));
+				}
+				else if(strcmp(Prefix.data(), EnglishText) != 0)
+				{
+					wprintf(L"WARNING: unexpected localized string for en-us.\n   Expected: '%hs'\n   Actual:   '%hs'", FindAndReplace(EnglishText, "\n", "\\n").c_str(), FindAndReplace(Prefix.data(), "\n", "\\n").c_str());
+				}
+			}
+			FreeLibrary(LibraryHandle);
+		}
+	}
+}
diff --git a/Engine/Extras/cl-filter/cl-filter.cpp b/Engine/Extras/cl-filter/cl-filter.cpp
deleted file mode 100644
