Allow printing out all the headers included from a source file by specifying -ShowIncludes on the UBT command line, or specifying <WindowsPlatform><ShowIncludes>true</ShowIncludes></WindowsPlatform> in BuildConfiguration.xml.

Also refactor argument parsing in cl-filter to accept case-insensitive argument names.

#rnx

diff --git a/Engine/Extras/Windows/cl-filter/cl-filter.cpp b/Engine/Extras/Windows/cl-filter/cl-filter.cpp
--- a/Engine/Extras/Windows/cl-filter/cl-filter.cpp
+++ b/Engine/Extras/Windows/cl-filter/cl-filter.cpp
@@ -16,11 +16,6 @@ void GetLocalizedIncludePrefixes(const wchar_t* CompilerPath, std::vector<std::v
 #define EXIT_CODE_SUCCESS 0
 #define EXIT_CODE_FAILURE -1
 
-#define OUTPUT_FILENAME_INDEX 1
-#define TIMING_FILENAME_INDEX 2
-#define COMPILER_FILENAME_INDEX_OFFSET 2
-#define MINIMUM_VALID_COMMAND_LINE_ARGS 4
-
 void PrintUsage()
 {
 	wprintf(L"Usage: cl-filter.exe -dependencies=<dependencies-file> -compiler=<compiler-file> [Optional Args] -- <child command line>\n");
@@ -28,57 +23,75 @@ void PrintUsage()
 	wprintf(L"\t-timing=<timing-file>\tThe file to write any timing data to.\n");
 }
 
-struct CommandLineArgs
+bool ParseArgumentValue(const wchar_t* Argument, const wchar_t* Prefix, size_t PrefixLen, const wchar_t** OutValue)
+{
+	if (_wcsnicmp(Argument, Prefix, PrefixLen) == 0)
+	{
+		*OutValue = Argument + PrefixLen;
+		return true;
+	}
+	return false;
+}
+
+template<size_t PrefixLen> 
+bool ParseArgumentValue(const wchar_t* Argument, const wchar_t(&Prefix)[PrefixLen], const wchar_t** OutValue)
 {
-	const wchar_t*	OutputFileName = nullptr;
-	const wchar_t*	TimingFileName = nullptr;
-	const wchar_t*	CompilerFileName = nullptr;
-};
+	return ParseArgumentValue(Argument, Prefix, PrefixLen - 1, OutValue);
+}
 
-CommandLineArgs* ParseCommandLineArgs(int ArgC, const wchar_t* ArgV[])
+int wmain(int ArgC, const wchar_t* ArgV[])
 {
-	CommandLineArgs* ParsedArgs = new CommandLineArgs();
-	for (int i = 1; i < ArgC; ++i)
+	// Parse the command line arguments
+	const wchar_t* OutputFileName = nullptr;
+	const wchar_t* TimingFileName = nullptr;
+	const wchar_t* CompilerFileName = nullptr;
+	bool bShowIncludes = false;
+
+	for (int Idx = 1; Idx < ArgC; Idx++)
 	{
 		// If this arg is the splitter, we're done parsing args.
-		if (wcscmp(ArgV[i], L"--") == 0)
+		const wchar_t* Arg = ArgV[Idx];
+		if (wcscmp(Arg, L"--") == 0)
 		{
 			break;
 		}
 
-		// If it starts with a -, parse the arg.
-		if (ArgV[i][0] == '-')
+		// Try to parse another argument
+		const wchar_t* Value;
+		if (ParseArgumentValue(Arg, L"-dependencies=", &Value))
 		{
-			if (wcsstr(ArgV[i], L"-dependencies=") != nullptr)
-			{
-				ParsedArgs->OutputFileName = ArgV[i] + 14;
-			}
-			else if (wcsstr(ArgV[i], L"-compiler=") != nullptr)
-			{
-				ParsedArgs->CompilerFileName = ArgV[i] + 10;
-			}
-			else if (wcsstr(ArgV[i], L"-timing=") != nullptr)
-			{
-				ParsedArgs->TimingFileName = ArgV[i] + 8;
-			}
+			OutputFileName = Value;
+			continue;
+		}
+		if (ParseArgumentValue(Arg, L"-compiler=", &Value))
+		{
+			CompilerFileName = Value;
+			continue;
+		}
+		if (ParseArgumentValue(Arg, L"-timing=", &Value))
+		{
+			TimingFileName = Value;
+			continue;
+		}
+		if (_wcsicmp(Arg, L"-showincludes") == 0)
+		{
+			bShowIncludes = true;
+			continue;
 		}
-	}
 
-	return ParsedArgs;
-}
+		// Output a warning if it doesn't match
+		wprintf(L"WARNING: Unknown argument '%s'\n", Arg);
+	}
 
-int wmain(int ArgC, const wchar_t* ArgV[])
-{
 	wchar_t* CmdLine = ::GetCommandLineW();
-	CommandLineArgs* ParsedArgs = ParseCommandLineArgs(ArgC, ArgV);
-	if (ParsedArgs->OutputFileName == nullptr)
+	if (OutputFileName == nullptr)
 	{
 		wprintf(L"ERROR: No output file name was specified! (%s)\n\n", CmdLine);
 		PrintUsage();
 		return EXIT_CODE_FAILURE;
 	}
 
-	if (ParsedArgs->CompilerFileName == nullptr)
+	if (CompilerFileName == nullptr)
 	{
 		wprintf(L"ERROR: No compiler file name was found! (%s)\n\n", CmdLine);
 		PrintUsage();
@@ -94,10 +107,10 @@ int wmain(int ArgC, const wchar_t* ArgV[])
 		return EXIT_CODE_FAILURE;
 	}
 	ChildCmdLine += 4;
-	
+
 	// Get all the possible localized string prefixes for /showIncludes output
 	std::vector<std::vector<char>> LocalizedIncludePrefixes;
-	GetLocalizedIncludePrefixes(ParsedArgs->CompilerFileName, LocalizedIncludePrefixes);
+	GetLocalizedIncludePrefixes(CompilerFileName, LocalizedIncludePrefixes);
 
 	// Create the child process
 	PROCESS_INFORMATION ProcessInfo;
@@ -146,20 +159,20 @@ int wmain(int ArgC, const wchar_t* ArgV[])
 	CloseHandle(StdErrWriteHandle);
 
 	// Delete the output file.
-	DeleteFileW(ParsedArgs->OutputFileName);
+	DeleteFileW(OutputFileName);
 
 	// Delete the timing file if needed.
-	if (ParsedArgs->TimingFileName != nullptr)
+	if (TimingFileName != nullptr)
 	{
-		DeleteFileW(ParsedArgs->TimingFileName);
+		DeleteFileW(TimingFileName);
 	}
 
 	// Get the path to a temporary output filename.
-	std::wstring TempOutputFileName(ParsedArgs->OutputFileName);
+	std::wstring TempOutputFileName(OutputFileName);
 	TempOutputFileName += L".tmp";
 
 	// Get the path to a temporary timing filename.
-	std::wstring TempTimingFileName(ParsedArgs->TimingFileName == nullptr ? L"" : ParsedArgs->TimingFileName);
+	std::wstring TempTimingFileName(TimingFileName == nullptr ? L"" : TimingFileName);
 	TempTimingFileName += L".tmp";
 
 	// Create a file to contain the dependency list.
@@ -172,7 +185,7 @@ int wmain(int ArgC, const wchar_t* ArgV[])
 
 	// If needed, create a file to contain unparsed timing info.
 	HANDLE TimingFile = INVALID_HANDLE_VALUE;
-	if (ParsedArgs->TimingFileName != nullptr)
+	if (TimingFileName != nullptr)
 	{
 		TimingFile = CreateFileW(TempTimingFileName.c_str(), GENERIC_WRITE, FILE_SHARE_READ, NULL, CREATE_ALWAYS, FILE_ATTRIBUTE_NORMAL, NULL);
 		if (TimingFile == INVALID_HANDLE_VALUE)
@@ -244,13 +257,19 @@ int wmain(int ArgC, const wchar_t* ArgV[])
 
 					DWORD BytesWritten;
 					WriteFile(OutputFile, Buffer + FileNameStart, (DWORD)(LineEnd - FileNameStart), &BytesWritten, NULL);
+
+					if (bShowIncludes)
+					{
+						WriteFile(GetStdHandle(STD_OUTPUT_HANDLE), Buffer + LineStart, (DWORD)(LineEnd - LineStart), &BytesWritten, NULL);
+					}
+
 					LineStart = LineEnd;
 					break;
 				}
 			}
 
 			// Handling timing info parsing if needed.
-			if (ParsedArgs->TimingFileName != nullptr && LineStart < LineEnd)
+			if (TimingFileName != nullptr && LineStart < LineEnd)
 			{
 				// See if we've moved into a block of timing information. Order is always Headers -> Classes -> Functions, so to keep for loops to a minimum,
 				//   handle by what our current state is.
@@ -296,18 +315,18 @@ int wmain(int ArgC, const wchar_t* ArgV[])
 
 	CloseHandle(OutputFile);
 
-	if (ExitCode == 0 && !MoveFileW(TempOutputFileName.c_str(), ParsedArgs->OutputFileName))
+	if (ExitCode == 0 && !MoveFileW(TempOutputFileName.c_str(), OutputFileName))
 	{
-		wprintf(L"ERROR: Unable to rename %s to %s\n", TempOutputFileName.c_str(), ParsedArgs->OutputFileName);
+		wprintf(L"ERROR: Unable to rename %s to %s\n", TempOutputFileName.c_str(), OutputFileName);
 		ExitCode = 1;
 	}
 
-	if (ParsedArgs->TimingFileName != nullptr)
+	if (TimingFileName != nullptr)
 	{
 		CloseHandle(TimingFile);
-		if (ExitCode == 0 && !MoveFileW(TempTimingFileName.c_str(), ParsedArgs->TimingFileName))
+		if (ExitCode == 0 && !MoveFileW(TempTimingFileName.c_str(), TimingFileName))
 		{
-			wprintf(L"ERROR: Unable to rename %s to %s\n", TempTimingFileName.c_str(), ParsedArgs->TimingFileName);
+			wprintf(L"ERROR: Unable to rename %s to %s\n", TempTimingFileName.c_str(), TimingFileName);
 			ExitCode = 1;
 		}
 	}
diff --git a/Engine/Source/Programs/UnrealBuildTool/Platform/Windows/UEBuildWindows.cs b/Engine/Source/Programs/UnrealBuildTool/Platform/Windows/UEBuildWindows.cs
--- a/Engine/Source/Programs/UnrealBuildTool/Platform/Windows/UEBuildWindows.cs
+++ b/Engine/Source/Programs/UnrealBuildTool/Platform/Windows/UEBuildWindows.cs
@@ -280,6 +280,13 @@ namespace UnrealBuildTool
 		[XmlConfigFile(Category = "WindowsPlatform")]
 		public bool bCompilerTrace = false;
 
+		/// <summary>
+		/// Print out files that are included by each source file
+		/// </summary>
+		[CommandLine("-ShowIncludes")]
+		[XmlConfigFile(Category = "WindowsPlatform")]
+		public bool bShowIncludes = false;
+
 		/// <summary>
 		/// Bundle a working version of dbghelp.dll with the application, and use this to generate minidumps. This works around a bug with the Windows 10 Fall Creators Update (1709)
 		/// where rich PE headers larger than a certain size would result in corrupt minidumps.
@@ -504,6 +511,11 @@ namespace UnrealBuildTool
 			get { return Inner.bCompilerTrace; }
 		}
 
+		public bool bShowIncludes
+		{
+			get { return Inner.bShowIncludes; }
+		}
+
 		public string GetVisualStudioCompilerVersionName()
 		{
 			return Inner.GetVisualStudioCompilerVersionName();
diff --git a/Engine/Source/Programs/UnrealBuildTool/Platform/Windows/VCToolChain.cs b/Engine/Source/Programs/UnrealBuildTool/Platform/Windows/VCToolChain.cs
--- a/Engine/Source/Programs/UnrealBuildTool/Platform/Windows/VCToolChain.cs
+++ b/Engine/Source/Programs/UnrealBuildTool/Platform/Windows/VCToolChain.cs
@@ -141,9 +141,6 @@ namespace UnrealBuildTool
 
 		void AppendCLArguments_Global(CppCompileEnvironment CompileEnvironment, List<string> Arguments)
 		{
-			// NOTE: Uncommenting this line will print includes as they are encountered by the preprocessor.  This can help with diagnosing include order problems.
-			//Arguments.Add("/showIncludes");
-
 			// Suppress generation of object code for unreferenced inline functions. Enabling this option is more standards compliant, and causes a big reduction
 			// in object file sizes (and link times) due to the amount of stuff we inline.
 			Arguments.Add("/Zc:inline");
@@ -1210,6 +1207,11 @@ namespace UnrealBuildTool
 				CommandArguments.Add(string.Format("-timing={0}", Utils.MakePathSafeToUseWithCommandLine(CompileAction.TimingFile.Location)));
 			}
 
+			if (Target.WindowsPlatform.bShowIncludes)
+			{
+				CommandArguments.Add("-showincludes");
+			}
+
 			CommandArguments.Add(String.Format("-compiler={0}", Utils.MakePathSafeToUseWithCommandLine(CompileAction.CommandPath)));
 			CommandArguments.Add("--");
 			CommandArguments.Add(Utils.MakePathSafeToUseWithCommandLine(CompileAction.CommandPath));
